Change Iris plugin status to Experimental.

diff --git a/Engine/Plugins/Experimental/Iris/Iris.uplugin b/Engine/Plugins/Experimental/Iris/Iris.uplugin
--- a/Engine/Plugins/Experimental/Iris/Iris.uplugin
+++ b/Engine/Plugins/Experimental/Iris/Iris.uplugin
@@ -12,7 +12,7 @@
     "SupportURL" : "",
     "EnabledByDefault" : true,
     "CanContainContent" : false,
-    "IsBetaVersion" : true,
+    "IsExperimentalVersion" : true,
     "Installed" : false,
     "Modules": [
       {
