Iris goes Beta.

diff --git a/Engine/Plugins/Experimental/Iris/Iris.uplugin b/Engine/Plugins/Experimental/Iris/Iris.uplugin
--- a/Engine/Plugins/Experimental/Iris/Iris.uplugin
+++ b/Engine/Plugins/Experimental/Iris/Iris.uplugin
@@ -12,7 +12,8 @@
 	"SupportURL": "",
 	"EnabledByDefault": false,
 	"CanContainContent": false,
-	"IsExperimentalVersion": true,
+	"IsBetaVersion": true,
+	"IsExperimentalVersion": false,
 	"Installed": false,
 	"Modules": [
 		{
